Fixing vehicle template by calling superclass functions correctly. Vehicle wasn't responding to controls or updating HUD because Super::BeginPlay() was never called.

diff --git a/Templates/TP_Vehicle/Source/TP_Vehicle/TP_VehiclePawn.cpp b/Templates/TP_Vehicle/Source/TP_Vehicle/TP_VehiclePawn.cpp
--- a/Templates/TP_Vehicle/Source/TP_Vehicle/TP_VehiclePawn.cpp
+++ b/Templates/TP_Vehicle/Source/TP_Vehicle/TP_VehiclePawn.cpp
@@ -106,6 +106,8 @@ ATP_VehiclePawn::ATP_VehiclePawn()
 
 void ATP_VehiclePawn::SetupPlayerInputComponent(class UInputComponent* InputComponent)
 {
+	Super::SetupPlayerInputComponent(InputComponent);
+
 	// set up gameplay key bindings
 	check(InputComponent);
 
@@ -178,6 +180,8 @@ void ATP_VehiclePawn::EnableIncarView(const bool bState, const bool bForce)
 
 void ATP_VehiclePawn::Tick(float Delta)
 {
+	Super::Tick(Delta);
+
 	// Setup the flag to say we are in reverse gear
 	bInReverseGear = GetVehicleMovement()->GetCurrentGear() < 0;
 	
@@ -208,6 +212,8 @@ void ATP_VehiclePawn::Tick(float Delta)
 
 void ATP_VehiclePawn::BeginPlay()
 {
+	Super::BeginPlay();
+
 	bool bEnableInCar = false;
 #ifdef HMD_INTGERATION
 	bEnableInCar = GEngine->HMDDevice.IsValid();	
